#ttp 330910 Mac: UE4EditorServices.app does not run from Matching UE4Editor.app location
#branch UE4
#project UE4EditorServices
#summary Give the UE4EditorServices the option to choose which editor executable to use to launch with or generate Xcode projects for.
#change Add a simple window, with pop-up menu for selecting which version of the editor to use for launching/building.
#change When handling services, if there are multiple editors installed then ask the user which to use, the dialog can be suppressed and can be forcibly shown (hold Alt when activating the service).

diff --git a/Engine/Source/Programs/Mac/UE4EditorServices/UE4EditorServices/UE4EditorServicesAppDelegate.h b/Engine/Source/Programs/Mac/UE4EditorServices/UE4EditorServices/UE4EditorServicesAppDelegate.h
--- a/Engine/Source/Programs/Mac/UE4EditorServices/UE4EditorServices/UE4EditorServicesAppDelegate.h
+++ b/Engine/Source/Programs/Mac/UE4EditorServices/UE4EditorServices/UE4EditorServicesAppDelegate.h
@@ -4,8 +4,17 @@
 
 @interface FUE4EditorServicesAppDelegate : NSObject <NSApplicationDelegate>
 
+@property (assign) IBOutlet NSWindow* Window;
+@property (assign) IBOutlet NSPopUpButton* EditorMenu;
+@property (assign) IBOutlet NSButton* CancelButton;
+@property (assign) IBOutlet NSButton* OKButton;
+@property (assign) IBOutlet NSButton* DontShow;
+@property (strong) NSMutableDictionary* DefaultApps;
+
+- (IBAction)onCancelButtonPressed:(id)sender;
+- (IBAction)onOKButtonPressed:(id)sender;
+
 - (void)launchGameService:(NSPasteboard *)pboard userData:(NSString *)userData error:(NSString **)error;
 - (void)generateXcodeFilesService:(NSPasteboard *)pboard userData:(NSString *)userData error:(NSString **)error;
-//- (void)generateAllXcodeFilesService:(NSPasteboard *)pboard userData:(NSString *)userData error:(NSString **)error;
 
 @end
diff --git a/Engine/Source/Programs/Mac/UE4EditorServices/UE4EditorServices/UE4EditorServicesAppDelegate.m b/Engine/Source/Programs/Mac/UE4EditorServices/UE4EditorServices/UE4EditorServicesAppDelegate.m
--- a/Engine/Source/Programs/Mac/UE4EditorServices/UE4EditorServices/UE4EditorServicesAppDelegate.m
+++ b/Engine/Source/Programs/Mac/UE4EditorServices/UE4EditorServices/UE4EditorServicesAppDelegate.m
@@ -4,10 +4,153 @@
 
 @implementation FUE4EditorServicesAppDelegate
 
+@synthesize Window;
+@synthesize EditorMenu;
+@synthesize OKButton;
+@synthesize CancelButton;
+@synthesize DontShow;
+@synthesize DefaultApps;
+
 - (void)applicationDidFinishLaunching:(NSNotification *)aNotification
 {
 	// Insert code here to initialize your application
 	[NSApp setServicesProvider:self];
+	DefaultApps = [NSMutableDictionary new];
+	NSDictionary* Current = (NSDictionary*)[[NSUserDefaults standardUserDefaults] objectForKey:@"DefaultApps"];
+	if(Current)
+	{
+		[DefaultApps addEntriesFromDictionary:Current];
+	}
+}
+
+/// construct menu item for app
+- (NSMenuItem *)menuItemForOpenWithForApplication:(NSString *)AppName URL:(NSURL *)AppURL
+{
+    NSMenuItem* AppItem = [NSMenuItem new];
+    [AppItem setTitle:AppName];
+    [AppItem setRepresentedObject:AppURL];
+	[AppItem setToolTip:[AppURL path]];
+	NSImage* Img = [[NSWorkspace sharedWorkspace] iconForFile:[AppURL path]];
+	[Img setSize:NSMakeSize(16, 16)];
+    [AppItem setImage:Img];
+    return AppItem;
+}
+
+- (void)openWithApplicationOtherSelected:(id)Sender
+{
+    NSOpenPanel* Panel = [NSOpenPanel openPanel];
+	
+    [Panel setAllowsMultipleSelection:NO];
+    [Panel setCanChooseFiles:YES];
+    [Panel setCanChooseDirectories:NO];
+	[Panel beginSheetModalForWindow: Window completionHandler:^(NSInteger result){
+		if (result == NSOKButton)
+		{
+			NSMenuItem* SelectedItem = nil;
+			bool AlreadyAddded = false;
+			for (NSMenuItem* Item in [[EditorMenu menu] itemArray])
+			{
+				if(Item && [Item representedObject] && [[[Item representedObject] path] compare:[[Panel URL] path] options:NSCaseInsensitiveSearch] == NSOrderedSame)
+				{
+					SelectedItem = Item;
+					AlreadyAddded = true;
+					break;
+				}
+			}
+			if(!AlreadyAddded)
+			{
+				NSString *appName = [[[[Panel URL] path] lastPathComponent] stringByDeletingPathExtension];
+				NSMenuItem *newAppItem = [self menuItemForOpenWithForApplication:appName URL:[Panel URL]];
+				[[EditorMenu menu] insertItem:newAppItem atIndex:2];
+				SelectedItem = newAppItem;
+			}
+			[Panel close];
+			[EditorMenu selectItem:SelectedItem];
+		}
+	}];
+}
+
+/// this method return open with menu for specified file
+- (NSMenu *)openWithMenuItemForURL:(NSURL *)FileURL
+{
+    NSMenu* SubMenu = [NSMenu new];
+    CFArrayRef AllApps = LSCopyApplicationURLsForURL((__bridge CFURLRef)FileURL, kLSRolesEditor);
+    NSMutableSet* AlreadyAdded = [NSMutableSet new];
+	if (AllApps != nil)
+	{
+		CFIndex count = CFArrayGetCount(AllApps);
+        //get and add default app
+        CFURLRef DefaultApp;
+        LSGetApplicationForURL((__bridge CFURLRef)FileURL, kLSRolesEditor, NULL, &DefaultApp);
+        if (DefaultApp)
+		{
+			NSString* DefaultAppPath = [(__bridge NSURL *)DefaultApp path];
+            NSString* DefaultAppName = [[[DefaultAppPath lastPathComponent] stringByDeletingPathExtension] stringByAppendingString:@" (default)"];
+            NSMenuItem* AppItem = [self menuItemForOpenWithForApplication:DefaultAppName URL:(__bridge NSURL *)DefaultApp];
+			[AlreadyAdded addObject:DefaultAppName];
+			[SubMenu addItem:AppItem];
+			[SubMenu addItem:[NSMenuItem separatorItem]];
+		}
+		else
+		{
+            NSLog(@"There is no default UE4 Editor for %@", FileURL);
+            NSMenuItem* NoneItem = [NSMenuItem new];
+            [NoneItem setTitle:@"<None>"];
+            [NoneItem setEnabled:NO];
+            [SubMenu addItem:NoneItem];
+            [SubMenu addItem:[NSMenuItem separatorItem]];
+        }
+        
+		if (count != 0)
+		{
+			for (int index = 0; index < count; ++index)
+			{
+				NSURL *appURL = (NSURL *)CFArrayGetValueAtIndex(AllApps, index);
+				if ([appURL isFileURL])
+				{
+					NSString* AppName = [[[appURL path] lastPathComponent] stringByDeletingPathExtension];
+					NSDictionary* Dictionary = [[NSBundle bundleWithPath:[appURL path]] infoDictionary];
+					if(Dictionary)
+					{
+						NSString* BundleId = (NSString*)[Dictionary valueForKey:@"CFBundleIdentifier"];
+						if (BundleId && ([BundleId compare:@"com.epicgames.UE4Editor" options:NSCaseInsensitiveSearch] == NSOrderedSame
+										 || [BundleId compare:@"com.epicgames.RocketEditor" options:NSCaseInsensitiveSearch] == NSOrderedSame))
+						{
+							if ([AlreadyAdded containsObject:AppName])
+							{
+								AppName = [AppName stringByAppendingFormat:@" (%@)", [Dictionary valueForKey:@"CFBundleVersion"]];
+							}
+							NSMenuItem* AppItem = [self menuItemForOpenWithForApplication:AppName URL:appURL];
+							[AlreadyAdded addObject:AppName];
+							[SubMenu addItem:AppItem];
+						}
+					}
+				}
+			}
+			[SubMenu addItem:[NSMenuItem separatorItem]];
+		}
+		
+        NSMenuItem* OtherAppItem = [NSMenuItem new];
+        [OtherAppItem setTitle:@"Other…"];
+        [OtherAppItem setTarget:self];
+        [OtherAppItem setAction:@selector(openWithApplicationOtherSelected:)];
+        [SubMenu addItem:OtherAppItem];
+		
+		CFRelease(AllApps);
+    }
+    return SubMenu;
+}
+
+- (IBAction)onCancelButtonPressed:(id)sender
+{
+	[NSApp stopModalWithCode:NSCancelButton];
+	[Window orderOut:self];
+}
+
+- (IBAction)onOKButtonPressed:(id)sender
+{
+	[NSApp stopModalWithCode:NSOKButton];
+	[Window orderOut:self];
 }
 
 - (void)launchGameService:(NSPasteboard *)pboard userData:(NSString *)userData error:(NSString **)error
@@ -25,7 +168,39 @@ - (void)launchGameService:(NSPasteboard *)pboard userData:(NSString *)userData e
 		
 		NSString* ProjectName = FilePath;
 		
-		NSURL* BundleURL = [[NSWorkspace sharedWorkspace] URLForApplicationToOpenURL:[NSURL fileURLWithPath: FilePath]];
+		NSURL* FileURL = [NSURL fileURLWithPath: FilePath];
+		
+		NSURL* BundleURL = nil;
+		
+		CFArrayRef AllAppURLs = LSCopyApplicationURLsForURL((__bridge CFURLRef)FileURL, kLSRolesAll);
+		bool AltDown = ([NSEvent modifierFlags] & NSAlternateKeyMask);
+		if([DefaultApps objectForKey:[FileURL absoluteString]] && !AltDown)
+		{
+			BundleURL = [NSURL URLWithString:(NSString*)[DefaultApps objectForKey:[FileURL absoluteString]]];
+		}
+		else if(CFArrayGetCount(AllAppURLs) > 1)
+		{
+			[EditorMenu setMenu:[self openWithMenuItemForURL:FileURL]];
+			[DontShow setState:NSOnState];
+			[Window setTitle:[NSString stringWithFormat:@"Launch %@ With...", [[FilePath lastPathComponent] stringByDeletingPathExtension]]];
+			[Window setRepresentedURL:FileURL];
+			NSInteger Code = [NSApp runModalForWindow:Window];
+			if(Code == NSOKButton)
+			{
+				BundleURL = (NSURL*)[[EditorMenu selectedItem] representedObject];
+				if([DontShow state] == NSOnState)
+				{
+					[DefaultApps setObject:[BundleURL absoluteString] forKey:[FileURL absoluteString]];
+					[[NSUserDefaults standardUserDefaults] setObject:(NSDictionary*)DefaultApps forKey:@"DefaultApps"];
+					[[NSUserDefaults standardUserDefaults] synchronize];
+				}
+			}
+		}
+		else
+		{
+			BundleURL = [[NSWorkspace sharedWorkspace] URLForApplicationToOpenURL:FileURL];
+		}
+		CFRelease(AllAppURLs);
 		
 		if(BundleURL)
 		{
@@ -76,11 +251,43 @@ - (void)generateXcodeFilesService:(NSPasteboard *)pboard userData:(NSString *)us
 		
 		NSString* FilePath = [FileArray objectAtIndex:0];
 		
-		NSURL* BundleURL = [[NSWorkspace sharedWorkspace] URLForApplicationToOpenURL:[NSURL fileURLWithPath: FilePath]];
+		NSURL* FileURL = [NSURL fileURLWithPath: FilePath];
+		
+		NSURL* BundleURL = nil;
+		
+		CFArrayRef AllAppURLs = LSCopyApplicationURLsForURL((__bridge CFURLRef)FileURL, kLSRolesAll);
+		bool AltDown = ([NSEvent modifierFlags] & NSAlternateKeyMask);
+		if([DefaultApps objectForKey:[FileURL absoluteString]] && !AltDown)
+		{
+			BundleURL = [NSURL URLWithString:(NSString*)[DefaultApps objectForKey:[FileURL absoluteString]]];
+		}
+		else if(CFArrayGetCount(AllAppURLs) > 1)
+		{
+			[EditorMenu setMenu:[self openWithMenuItemForURL:FileURL]];
+			[DontShow setState:NSOnState];
+			[Window setTitle:[NSString stringWithFormat:@"Generate %@ Xcode Project With...", [[FilePath lastPathComponent] stringByDeletingPathExtension]]];
+			[Window setRepresentedURL:FileURL];
+			NSInteger Code = [NSApp runModalForWindow:Window];
+			if(Code == NSOKButton)
+			{
+				BundleURL = (NSURL*)[[EditorMenu selectedItem] representedObject];
+				if([DontShow state] == NSOnState)
+				{
+					[DefaultApps setObject:[BundleURL absoluteString] forKey:[FileURL absoluteString]];
+					[[NSUserDefaults standardUserDefaults] setObject:(NSDictionary*)DefaultApps forKey:@"DefaultApps"];
+					[[NSUserDefaults standardUserDefaults] synchronize];
+				}
+			}
+		}
+		else
+		{
+			BundleURL = [[NSWorkspace sharedWorkspace] URLForApplicationToOpenURL:[NSURL fileURLWithPath: FilePath]];
+		}
+		CFRelease(AllAppURLs);
 		
 		NSString* FolderPath = nil;
 		NSString* ScriptPath = nil;
-
+		
 		if(BundleURL)
 		{
 			NSString* CurrentPath = [[BundleURL path] stringByDeletingLastPathComponent];
