Added Windows Main file for UE4EditorServices so that GUBP doesn't complain

diff --git a/Engine/Source/Programs/Mac/UE4EditorServices/Private/Windows/UE4EditorServicesMain.cpp b/Engine/Source/Programs/Mac/UE4EditorServices/Private/Windows/UE4EditorServicesMain.cpp
new file mode 100644
--- /dev/null
+++ b/Engine/Source/Programs/Mac/UE4EditorServices/Private/Windows/UE4EditorServicesMain.cpp
@@ -0,0 +1,12 @@
+// Copyright 1998-2014 Epic Games, Inc. All Rights Reserved.
+
+#include "Core.h"
+#include "RequiredProgramMainCPPInclude.h"
+#include "ModuleManager.h"
+
+IMPLEMENT_APPLICATION(UE4EditorServices, "UE4EditorServices");
+
+int WINAPI WinMain(HINSTANCE hInstance, HINSTANCE hPrevInstance, LPSTR lpCmdLine, int ShowCmd)
+{
+	return 0;
+}
