Added CFBundleShortVersionString to UE4EditorServices and edited Mac FixDependencyFiles script to copy services app to ~/Library/Services

diff --git a/Engine/Build/BatchFiles/Mac/FixDependencyFiles.sh b/Engine/Build/BatchFiles/Mac/FixDependencyFiles.sh
--- a/Engine/Build/BatchFiles/Mac/FixDependencyFiles.sh
+++ b/Engine/Build/BatchFiles/Mac/FixDependencyFiles.sh
@@ -22,3 +22,8 @@ if [ $(ls casetest* | wc -l) -gt 1 ]; then
     done
 fi
 rm -f casetest*
+
+# Copy UE4EditorServices to ~/Library/Services
+if [ ! -d ~/Library/Services/UE4EditorServices.app ]; then
+	cp -r ../../../Binaries/Mac/UE4EditorServices.app ~/Library/Services/UE4EditorServices.app
+fi
diff --git a/Engine/Source/Programs/Mac/UE4EditorServices/Resources/Mac/Info.plist b/Engine/Source/Programs/Mac/UE4EditorServices/Resources/Mac/Info.plist
--- a/Engine/Source/Programs/Mac/UE4EditorServices/Resources/Mac/Info.plist
+++ b/Engine/Source/Programs/Mac/UE4EditorServices/Resources/Mac/Info.plist
@@ -33,6 +33,8 @@
 	<string>${EXECUTABLE_NAME}</string>
 	<key>CFBundlePackageType</key>
 	<string>APPL</string>
+	<key>CFBundleShortVersionString</key>
+	<string>1.1</string>
 	<key>CFBundleSignature</key>
 	<string>????</string>
 	<key>CFBundleVersion</key>
