Merging CL #2100598 to Main (Fixed showing of engine version selector window without a nib file in UE4EditorServices)

diff --git a/Engine/Source/Programs/Mac/UE4EditorServices/Private/Mac/UE4EditorServicesAppDelegate.cpp b/Engine/Source/Programs/Mac/UE4EditorServices/Private/Mac/UE4EditorServicesAppDelegate.cpp
--- a/Engine/Source/Programs/Mac/UE4EditorServices/Private/Mac/UE4EditorServicesAppDelegate.cpp
+++ b/Engine/Source/Programs/Mac/UE4EditorServices/Private/Mac/UE4EditorServicesAppDelegate.cpp
@@ -8,11 +8,6 @@
 
 @implementation FUE4EditorServicesAppDelegate
 
-@synthesize Window;
-@synthesize EditorMenu;
-@synthesize OKButton;
-@synthesize CancelButton;
-
 - (void)applicationDidFinishLaunching:(NSNotification *)Notification
 {
 	[NSApp setServicesProvider:self];
@@ -226,13 +221,13 @@
 - (IBAction)onCancelButtonPressed:(id)Sender
 {
 	[NSApp stopModalWithCode:NSCancelButton];
-	[Window orderOut:self];
+	[Window close];
 }
 
 - (IBAction)onOKButtonPressed:(id)Sender
 {
 	[NSApp stopModalWithCode:NSOKButton];
-	[Window orderOut:self];
+	[Window close];
 }
 
 - (NSMenuItem*)addEnginePath:(NSString*)EnginePath toEditorMenu:(NSMenu*)Menu withDescription:(NSString*)Suffix
@@ -290,10 +285,9 @@
     {
 		NSURL* FileURL = [NSURL fileURLWithPath: [[PBoard propertyListForType:NSFilenamesPboardType] objectAtIndex:0]];
 
-		NSString* FileDefaultEnginePath = [self enginePathForAppBundle:[self defaultAppForUProjectFile:FileURL]];
 		NSString* GlobalDefaultEnginePath = [self enginePathForAppBundle:[self defaultAppForUProjectFiles]];
 		NSString* RecommendedEnginePath = [self enginePathForAppBundle:[self recommendedAppForUProjectFile:FileURL]];
-		NSString* SelectedEnginePath = FileDefaultEnginePath ? FileDefaultEnginePath : RecommendedEnginePath;
+		NSString* SelectedEnginePath = [self findEngineForUProjectFile:FileURL];
 
 		NSMenu* SubMenu = [NSMenu new];
 		NSMenuItem* SelectedItem = NULL;
@@ -327,9 +321,39 @@
 			}
 		}
 
+		Window = [[NSWindow alloc] initWithContentRect:NSMakeRect(0, 0, 433, 102) styleMask:NSTitledWindowMask backing:NSBackingStoreBuffered defer:NO];
+
+		NSPopUpButton* EditorMenu = [[NSPopUpButton alloc] initWithFrame:NSMakeRect(18, 58, 398, 26) pullsDown:NO];
 		[EditorMenu setMenu:SubMenu];
 		[EditorMenu selectItem:SelectedItem];
 
+		NSButton* OKButton = [NSButton new];
+		[OKButton setButtonType:NSMomentaryLightButton];
+		[OKButton setTitle:@"OK"];
+		[OKButton setFrame:NSMakeRect(337, 13, 82, 32)];
+		[OKButton setImagePosition:NSNoImage];
+		[OKButton setBezelStyle:NSRoundedBezelStyle];
+		[OKButton setFont:[NSFont systemFontOfSize:0]];
+		[OKButton setAction:@selector(onOKButtonPressed:)];
+		[OKButton setTarget:self];
+
+		NSButton* CancelButton = [NSButton new];
+		[CancelButton setButtonType:NSMomentaryLightButton];
+		[CancelButton setTitle:@"Cancel"];
+		[CancelButton setFrame:NSMakeRect(255, 13, 82, 32)];
+		[CancelButton setImagePosition:NSNoImage];
+		[CancelButton setBezelStyle:NSRoundedBezelStyle];
+		[CancelButton setFont:[NSFont systemFontOfSize:0]];
+		[CancelButton setAction:@selector(onCancelButtonPressed:)];
+		[CancelButton setTarget:self];
+
+		[[Window contentView] addSubview:EditorMenu];
+		[[Window contentView] addSubview:OKButton];
+		[[Window contentView] addSubview:CancelButton];
+
+		[Window center];
+		[NSApp activateIgnoringOtherApps:YES];
+
 		NSInteger Result = [NSApp runModalForWindow:Window];
 		if (Result == NSOKButton)
 		{
diff --git a/Engine/Source/Programs/Mac/UE4EditorServices/Private/Mac/UE4EditorServicesAppDelegate.h b/Engine/Source/Programs/Mac/UE4EditorServices/Private/Mac/UE4EditorServicesAppDelegate.h
--- a/Engine/Source/Programs/Mac/UE4EditorServices/Private/Mac/UE4EditorServicesAppDelegate.h
+++ b/Engine/Source/Programs/Mac/UE4EditorServices/Private/Mac/UE4EditorServicesAppDelegate.h
@@ -3,11 +3,9 @@
 #import <Cocoa/Cocoa.h>
 
 @interface FUE4EditorServicesAppDelegate : NSObject <NSApplicationDelegate>
-
-@property (assign) IBOutlet NSWindow* Window;
-@property (assign) IBOutlet NSPopUpButton* EditorMenu;
-@property (assign) IBOutlet NSButton* CancelButton;
-@property (assign) IBOutlet NSButton* OKButton;
+{
+	NSWindow* Window;
+}
 
 - (IBAction)onCancelButtonPressed:(id)Sender;
 - (IBAction)onOKButtonPressed:(id)Sender;
diff --git a/Engine/Source/Programs/Mac/UE4EditorServices/Resources/Mac/Info.plist b/Engine/Source/Programs/Mac/UE4EditorServices/Resources/Mac/Info.plist
--- a/Engine/Source/Programs/Mac/UE4EditorServices/Resources/Mac/Info.plist
+++ b/Engine/Source/Programs/Mac/UE4EditorServices/Resources/Mac/Info.plist
@@ -38,11 +38,11 @@
 	<key>CFBundlePackageType</key>
 	<string>APPL</string>
 	<key>CFBundleShortVersionString</key>
-	<string>1.4</string>
+	<string>1.5</string>
 	<key>CFBundleSignature</key>
 	<string>????</string>
 	<key>CFBundleVersion</key>
-	<string>1.4</string>
+	<string>1.5</string>
 	<key>LSApplicationCategoryType</key>
 	<string>public.app-category.developer-tools</string>
 	<key>LSBackgroundOnly</key>
