Merging CL #2100803 to Main (Added "Choose UE4 Editor..." title to Select Unreal Engine Version window on Mac)

diff --git a/Engine/Source/Programs/Mac/UE4EditorServices/Private/Mac/UE4EditorServicesAppDelegate.cpp b/Engine/Source/Programs/Mac/UE4EditorServices/Private/Mac/UE4EditorServicesAppDelegate.cpp
--- a/Engine/Source/Programs/Mac/UE4EditorServices/Private/Mac/UE4EditorServicesAppDelegate.cpp
+++ b/Engine/Source/Programs/Mac/UE4EditorServices/Private/Mac/UE4EditorServicesAppDelegate.cpp
@@ -351,6 +351,8 @@
 		[[Window contentView] addSubview:OKButton];
 		[[Window contentView] addSubview:CancelButton];
 
+		[Window setTitle:@"Choose UE4 Editor..."];
+
 		[Window center];
 		[NSApp activateIgnoringOtherApps:YES];
 
