[INTEGRATE] Merging using UE4-LauncherReleasesTo-UE4 all up to cl 2111815

diff --git a/Engine/Source/Developer/DesktopPlatform/Private/Mac/DesktopPlatformMac.cpp b/Engine/Source/Developer/DesktopPlatform/Private/Mac/DesktopPlatformMac.cpp
--- a/Engine/Source/Developer/DesktopPlatform/Private/Mac/DesktopPlatformMac.cpp
+++ b/Engine/Source/Developer/DesktopPlatform/Private/Mac/DesktopPlatformMac.cpp
@@ -618,7 +618,7 @@ bool FDesktopPlatformMac::VerifyFileAssociations()
 		NSBundle* GlobalDefaultAppBundle = [NSBundle bundleWithURL:(__bridge NSURL*)GlobalDefaultAppURL];
 		CFRelease(GlobalDefaultAppURL);
 
-		if ([[GlobalDefaultAppBundle bundleIdentifier] isEqualToString:@"com.epicgames.UE4Editor"] || [[GlobalDefaultAppBundle bundleIdentifier] isEqualToString:@"com.epicgames.UE4EditorServices"])
+		if ([[GlobalDefaultAppBundle bundleIdentifier] isEqualToString:@"com.epicgames.UE4EditorServices"])
 		{
 			return true;
 		}
diff --git a/Engine/Source/Programs/Mac/UE4EditorServices/Private/Mac/UE4EditorServicesAppDelegate.cpp b/Engine/Source/Programs/Mac/UE4EditorServices/Private/Mac/UE4EditorServicesAppDelegate.cpp
--- a/Engine/Source/Programs/Mac/UE4EditorServices/Private/Mac/UE4EditorServicesAppDelegate.cpp
+++ b/Engine/Source/Programs/Mac/UE4EditorServices/Private/Mac/UE4EditorServicesAppDelegate.cpp
@@ -71,7 +71,7 @@
 	return [NSBundle mainBundle];
 }
 
-- (NSString*)findEngineForUProjectFile:(NSURL*)FileURL
+- (NSString*)defaultEngineForUProjectFile:(NSURL*)FileURL
 {
 	NSString* FileDefaultEnginePath = [self enginePathForAppBundle:[self defaultAppForUProjectFile:FileURL]];
 	NSString* GlobalDefaultEnginePath = [self enginePathForAppBundle:[self defaultAppForUProjectFiles]];
@@ -82,20 +82,161 @@
 	return FileDefaultEnginePath;
 }
 
+- (IBAction)onCancelButtonPressed:(id)Sender
+{
+	[NSApp stopModalWithCode:NSCancelButton];
+	[Window close];
+}
+
+- (IBAction)onOKButtonPressed:(id)Sender
+{
+	[NSApp stopModalWithCode:NSOKButton];
+	[Window close];
+}
+
+- (FString)showEngineSelectionWindowForURL:(NSURL*)FileURL
+{
+	TMap<FString, FString> Installations;
+	FDesktopPlatformModule::Get()->EnumerateEngineInstallations(Installations);
+	if (Installations.Num() == 0)
+	{
+		FPlatformMisc::MessageBoxExt(EAppMsgType::Ok, TEXT("Unreal Engine is not installed"), TEXT("Error"));
+		return TEXT("");
+	}
+
+	NSString* GlobalDefaultEnginePath = [self enginePathForAppBundle:[self defaultAppForUProjectFiles]];
+	NSString* RecommendedEnginePath = [self enginePathForAppBundle:[self recommendedAppForUProjectFile:FileURL]];
+	NSString* SelectedEnginePath = [self defaultEngineForUProjectFile:FileURL];
+
+	NSMenu* SubMenu = [NSMenu new];
+	NSMenuItem* SelectedItem = NULL;
+
+	if (GlobalDefaultEnginePath)
+	{
+		SelectedItem = [self addEnginePath:GlobalDefaultEnginePath toEditorMenu:SubMenu withDescription:@"(default)"];
+	}
+
+	if (GlobalDefaultEnginePath)
+	{
+		[SubMenu addItem:[NSMenuItem separatorItem]];
+	}
+
+	for (TMap<FString, FString>::TConstIterator Iter(Installations); Iter; ++Iter)
+	{
+		NSString* EnginePath = Iter.Value().GetNSString();
+		if (![EnginePath isEqualToString:GlobalDefaultEnginePath])
+		{
+			NSMenuItem* MenuItem = [self addEnginePath:EnginePath toEditorMenu:SubMenu withDescription:[EnginePath isEqualToString:RecommendedEnginePath] ? @"(recommended)" : NULL];
+			if ([EnginePath isEqualToString:SelectedEnginePath])
+			{
+				SelectedItem = MenuItem;
+			}
+		}
+	}
+
+	Window = [[NSWindow alloc] initWithContentRect:NSMakeRect(0, 0, 433, 102) styleMask:NSTitledWindowMask backing:NSBackingStoreBuffered defer:NO];
+
+	NSPopUpButton* EditorMenu = [[NSPopUpButton alloc] initWithFrame:NSMakeRect(18, 58, 398, 26) pullsDown:NO];
+	[EditorMenu setMenu:SubMenu];
+	[EditorMenu selectItem:SelectedItem];
+
+	NSButton* OKButton = [NSButton new];
+	[OKButton setButtonType:NSMomentaryLightButton];
+	[OKButton setTitle:@"OK"];
+	[OKButton setFrame:NSMakeRect(337, 13, 82, 32)];
+	[OKButton setImagePosition:NSNoImage];
+	[OKButton setBezelStyle:NSRoundedBezelStyle];
+	[OKButton setFont:[NSFont systemFontOfSize:0]];
+	[OKButton setAction:@selector(onOKButtonPressed:)];
+	[OKButton setTarget:self];
+
+	NSButton* CancelButton = [NSButton new];
+	[CancelButton setButtonType:NSMomentaryLightButton];
+	[CancelButton setTitle:@"Cancel"];
+	[CancelButton setFrame:NSMakeRect(255, 13, 82, 32)];
+	[CancelButton setImagePosition:NSNoImage];
+	[CancelButton setBezelStyle:NSRoundedBezelStyle];
+	[CancelButton setFont:[NSFont systemFontOfSize:0]];
+	[CancelButton setAction:@selector(onCancelButtonPressed:)];
+	[CancelButton setTarget:self];
+
+	[[Window contentView] addSubview:EditorMenu];
+	[[Window contentView] addSubview:OKButton];
+	[[Window contentView] addSubview:CancelButton];
+
+	[Window setTitle:@"Choose UE4 Editor..."];
+
+	[Window center];
+	[NSApp activateIgnoringOtherApps:YES];
+
+	NSInteger Result = [NSApp runModalForWindow:Window];
+	if (Result == NSOKButton)
+	{
+		SelectedEnginePath = (NSString*)[[EditorMenu selectedItem] representedObject];
+		for (TMap<FString, FString>::TConstIterator Iter(Installations); Iter; ++Iter)
+		{
+			if ([Iter.Value().GetNSString() isEqualToString:SelectedEnginePath])
+			{
+				return Iter.Key();
+			}
+		}
+	}
+
+	return TEXT("");
+}
+
+- (NSString*)findEngineForUProjectFile:(NSURL*)FileURL
+{
+	// First, check if this file is associated to a specific build in Finder. If it is, we use this build no matter what's set in .uproject's EngineAssociation
+	NSString* FileDefaultEnginePath = [self enginePathForAppBundle:[self defaultAppForUProjectFile:FileURL]];
+	NSString* GlobalDefaultEnginePath = [self enginePathForAppBundle:[self defaultAppForUProjectFiles]];
+	if (FileDefaultEnginePath && ![FileDefaultEnginePath isEqualToString:GlobalDefaultEnginePath] && [[NSFileManager defaultManager] fileExistsAtPath:FileDefaultEnginePath])
+	{
+		return FileDefaultEnginePath;
+	}
+
+	TMap<FString, FString> Installations;
+	FDesktopPlatformModule::Get()->EnumerateEngineInstallations(Installations);
+
+	// Now, read EngineAssociation from .uproject file and use the build that the project is associated with if it's installed
+	FString EngineAssociation;
+	if (FDesktopPlatformModule::Get()->GetEngineIdentifierForProject([FileURL path], EngineAssociation))
+	{
+		for (TMap<FString, FString>::TConstIterator Iter(Installations); Iter; ++Iter)
+		{
+			if (Iter.Key() == EngineAssociation)
+			{
+				return Iter.Value().GetNSString();
+			}
+		}
+	}
+
+	// The build this file is associated with is not installed, so ask user to select one
+	FString EngineID = [self showEngineSelectionWindowForURL:FileURL];
+	if (!EngineID.IsEmpty())
+	{
+		return Installations[EngineID].GetNSString();
+	}
+
+	return nil;
+}
+
 - (BOOL)application:(NSApplication*)Application openFile:(NSString*)Filename
 {
 	NSURL* FileURL = [NSURL fileURLWithPath:Filename];
 	NSString* EnginePath = [self findEngineForUProjectFile:FileURL];
-	NSURL* EditorBundleURL = [NSURL fileURLWithPath:[EnginePath stringByAppendingPathComponent:@"Engine/Binaries/Mac/UE4Editor.app"]];
-
-	if (EditorBundleURL)
+	if (EnginePath)
 	{
-		NSDictionary* Configuration = [NSDictionary dictionaryWithObject:[NSArray arrayWithObjects:[FileURL path], nil] forKey:NSWorkspaceLaunchConfigurationArguments];
-		NSError* LaunchError = NULL;
-		NSRunningApplication* NewInstance = [[NSWorkspace sharedWorkspace] launchApplicationAtURL:EditorBundleURL options:(NSWorkspaceLaunchOptions)(NSWorkspaceLaunchAsync|NSWorkspaceLaunchNewInstance) configuration:Configuration error:&LaunchError];
-		if (NewInstance)
+		NSURL* EditorBundleURL = [NSURL fileURLWithPath:[EnginePath stringByAppendingPathComponent:@"Engine/Binaries/Mac/UE4Editor.app"]];
+		if (EditorBundleURL)
 		{
-			return YES;
+			NSDictionary* Configuration = [NSDictionary dictionaryWithObject:[NSArray arrayWithObjects:[FileURL path], nil] forKey:NSWorkspaceLaunchConfigurationArguments];
+			NSError* LaunchError = NULL;
+			NSRunningApplication* NewInstance = [[NSWorkspace sharedWorkspace] launchApplicationAtURL:EditorBundleURL options:(NSWorkspaceLaunchOptions)(NSWorkspaceLaunchAsync|NSWorkspaceLaunchNewInstance) configuration:Configuration error:&LaunchError];
+			if (NewInstance)
+			{
+				return YES;
+			}
 		}
 	}
 
@@ -110,28 +251,30 @@
     {
 		NSURL* FileURL = [NSURL fileURLWithPath: [[PBoard propertyListForType:NSFilenamesPboardType] objectAtIndex:0]];
 		NSString* EnginePath = [self findEngineForUProjectFile:FileURL];
-		NSURL* EditorBundleURL = [NSURL fileURLWithPath:[EnginePath stringByAppendingPathComponent:@"Engine/Binaries/Mac/UE4Editor.app"]];
-
-		if (EditorBundleURL)
+		if (EnginePath)
 		{
-			NSDictionary* Configuration = [NSDictionary dictionaryWithObject:[NSArray arrayWithObjects:[FileURL path], @"-game", nil] forKey:NSWorkspaceLaunchConfigurationArguments];
-			NSError* LaunchError = NULL;
-			NSRunningApplication* NewInstance = [[NSWorkspace sharedWorkspace] launchApplicationAtURL:EditorBundleURL options:(NSWorkspaceLaunchOptions)(NSWorkspaceLaunchAsync|NSWorkspaceLaunchNewInstance) configuration:Configuration error:&LaunchError];
-			if (!NewInstance)
+			NSURL* EditorBundleURL = [NSURL fileURLWithPath:[EnginePath stringByAppendingPathComponent:@"Engine/Binaries/Mac/UE4Editor.app"]];
+			if (EditorBundleURL)
 			{
-				if (LaunchError)
-				{
-					UnrealError = [LaunchError localizedDescription];
-				}
-				else
+				NSDictionary* Configuration = [NSDictionary dictionaryWithObject:[NSArray arrayWithObjects:[FileURL path], @"-game", nil] forKey:NSWorkspaceLaunchConfigurationArguments];
+				NSError* LaunchError = NULL;
+				NSRunningApplication* NewInstance = [[NSWorkspace sharedWorkspace] launchApplicationAtURL:EditorBundleURL options:(NSWorkspaceLaunchOptions)(NSWorkspaceLaunchAsync|NSWorkspaceLaunchNewInstance) configuration:Configuration error:&LaunchError];
+				if (!NewInstance)
 				{
-					UnrealError = @"Failed to run a copy of the game on this machine.";
+					if (LaunchError)
+					{
+						UnrealError = [LaunchError localizedDescription];
+					}
+					else
+					{
+						UnrealError = @"Failed to run a copy of the game on this machine.";
+					}
 				}
 			}
-		}
-		else
-		{
-			UnrealError = @"No application to open the project file available.";
+			else
+			{
+				UnrealError = @"No application to open the project file available.";
+			}
 		}
 	}
 	else
@@ -153,59 +296,61 @@
     {
 		NSURL* FileURL = [NSURL fileURLWithPath: [[PBoard propertyListForType:NSFilenamesPboardType] objectAtIndex:0]];
 		NSString* EnginePath = [self findEngineForUProjectFile:FileURL];
-
-		NSString* ScriptPath = [EnginePath stringByAppendingPathComponent:@"Engine/Build/BatchFiles/Mac/GenerateProjectFiles.sh"];
-		if (![[NSFileManager defaultManager] fileExistsAtPath:ScriptPath])
+		if (EnginePath)
 		{
-			ScriptPath = [EnginePath stringByAppendingPathComponent:@"Engine/Build/BatchFiles/Mac/RocketGenerateProjectFiles.sh"];
-		}
+			NSString* ScriptPath = [EnginePath stringByAppendingPathComponent:@"Engine/Build/BatchFiles/Mac/GenerateProjectFiles.sh"];
+			if (![[NSFileManager defaultManager] fileExistsAtPath:ScriptPath])
+			{
+				ScriptPath = [EnginePath stringByAppendingPathComponent:@"Engine/Build/BatchFiles/Mac/RocketGenerateProjectFiles.sh"];
+			}
 
-		if ([[NSFileManager defaultManager] fileExistsAtPath:ScriptPath])
-		{
-			if ([[NSWorkspace sharedWorkspace] launchAppWithBundleIdentifier:@"com.apple.terminal" options:0 additionalEventParamDescriptor:nil launchIdentifier:nil])
+			if ([[NSFileManager defaultManager] fileExistsAtPath:ScriptPath])
 			{
-				NSString* FullFolderPath = [[ScriptPath stringByDeletingLastPathComponent] stringByResolvingSymlinksInPath];
-				NSString* FullScriptPath = [ScriptPath stringByResolvingSymlinksInPath];
-				NSString* Command = [NSString stringWithFormat:@"cd \"%@\" \n sh \"%@\" -project=\"%@\" -game\n logout\n", FullFolderPath, FullScriptPath, [FileURL path]];
-
-				const char* UTF8Script = [Command UTF8String];
-
-				const char *TerminalBundleID = "com.apple.terminal";
-
-				// Build event
-				AppleEvent Event, Reply;
-				OSStatus Status = AEBuildAppleEvent(kAECoreSuite, kAEDoScript,
-										typeApplicationBundleID,
-										TerminalBundleID, strlen(TerminalBundleID),
-										kAutoGenerateReturnID,
-										kAnyTransactionID,
-										&Event, NULL,
-										"'----':utf8(@)", [Command length],
-										UTF8Script);
-				if (Status == noErr)
+				if ([[NSWorkspace sharedWorkspace] launchAppWithBundleIdentifier:@"com.apple.terminal" options:0 additionalEventParamDescriptor:nil launchIdentifier:nil])
 				{
-					// Send event and check for any Apple Event Manager errors
-					Status = AESendMessage(&Event, &Reply, kAENoReply, kAEDefaultTimeout);
-					AEDisposeDesc(&Event);
-					if (Status != noErr)
+					NSString* FullFolderPath = [[ScriptPath stringByDeletingLastPathComponent] stringByResolvingSymlinksInPath];
+					NSString* FullScriptPath = [ScriptPath stringByResolvingSymlinksInPath];
+					NSString* Command = [NSString stringWithFormat:@"cd \"%@\" \n sh \"%@\" -project=\"%@\" -game\n logout\n", FullFolderPath, FullScriptPath, [FileURL path]];
+
+					const char* UTF8Script = [Command UTF8String];
+
+					const char *TerminalBundleID = "com.apple.terminal";
+
+					// Build event
+					AppleEvent Event, Reply;
+					OSStatus Status = AEBuildAppleEvent(kAECoreSuite, kAEDoScript,
+											typeApplicationBundleID,
+											TerminalBundleID, strlen(TerminalBundleID),
+											kAutoGenerateReturnID,
+											kAnyTransactionID,
+											&Event, NULL,
+											"'----':utf8(@)", [Command length],
+											UTF8Script);
+					if (Status == noErr)
+					{
+						// Send event and check for any Apple Event Manager errors
+						Status = AESendMessage(&Event, &Reply, kAENoReply, kAEDefaultTimeout);
+						AEDisposeDesc(&Event);
+						if (Status != noErr)
+						{
+							UnrealError = @"Couldn't tell Terminal to generate project files.";
+						}
+					}
+					else
 					{
 						UnrealError = @"Couldn't tell Terminal to generate project files.";
 					}
 				}
 				else
 				{
-					UnrealError = @"Couldn't tell Terminal to generate project files.";
+					UnrealError = @"Failed to open Terminal while trying to generate project files.";
 				}
 			}
 			else
 			{
-				UnrealError = @"Failed to open Terminal while trying to generate project files.";
+				UnrealError = @"No application to generate project files available.";
 			}
 		}
-		else
-		{
-			UnrealError = @"No application to generate project files available.";
-		}
 	}
 	else
 	{
@@ -218,18 +363,6 @@
 	}
 }
 
-- (IBAction)onCancelButtonPressed:(id)Sender
-{
-	[NSApp stopModalWithCode:NSCancelButton];
-	[Window close];
-}
-
-- (IBAction)onOKButtonPressed:(id)Sender
-{
-	[NSApp stopModalWithCode:NSOKButton];
-	[Window close];
-}
-
 - (NSMenuItem*)addEnginePath:(NSString*)EnginePath toEditorMenu:(NSMenu*)Menu withDescription:(NSString*)Suffix
 {
 	NSFileManager* FileManager = [NSFileManager defaultManager];
@@ -284,95 +417,10 @@
 	if ([[PBoard types] containsObject:NSFilenamesPboardType])
     {
 		NSURL* FileURL = [NSURL fileURLWithPath: [[PBoard propertyListForType:NSFilenamesPboardType] objectAtIndex:0]];
-
-		NSString* GlobalDefaultEnginePath = [self enginePathForAppBundle:[self defaultAppForUProjectFiles]];
-		NSString* RecommendedEnginePath = [self enginePathForAppBundle:[self recommendedAppForUProjectFile:FileURL]];
-		NSString* SelectedEnginePath = [self findEngineForUProjectFile:FileURL];
-
-		NSMenu* SubMenu = [NSMenu new];
-		NSMenuItem* SelectedItem = NULL;
-
-		if (GlobalDefaultEnginePath)
-		{
-			SelectedItem = [self addEnginePath:GlobalDefaultEnginePath toEditorMenu:SubMenu withDescription:@"(default)"];
-		}
-
-		TMap<FString, FString> Installations;
-		FDesktopPlatformModule::Get()->EnumerateEngineInstallations(Installations);
-
-		if (Installations.Num() > 0)
-		{
-			if (GlobalDefaultEnginePath)
-			{
-				[SubMenu addItem:[NSMenuItem separatorItem]];
-			}
-
-			for (TMap<FString, FString>::TConstIterator Iter(Installations); Iter; ++Iter)
-			{
-				NSString* EnginePath = Iter.Value().GetNSString();
-				if (![EnginePath isEqualToString:GlobalDefaultEnginePath])
-				{
-					NSMenuItem* MenuItem = [self addEnginePath:EnginePath toEditorMenu:SubMenu withDescription:[EnginePath isEqualToString:RecommendedEnginePath] ? @"(recommended)" : NULL];
-					if ([EnginePath isEqualToString:SelectedEnginePath])
-					{
-						SelectedItem = MenuItem;
-					}
-				}
-			}
-		}
-
-		Window = [[NSWindow alloc] initWithContentRect:NSMakeRect(0, 0, 433, 102) styleMask:NSTitledWindowMask backing:NSBackingStoreBuffered defer:NO];
-
-		NSPopUpButton* EditorMenu = [[NSPopUpButton alloc] initWithFrame:NSMakeRect(18, 58, 398, 26) pullsDown:NO];
-		[EditorMenu setMenu:SubMenu];
-		[EditorMenu selectItem:SelectedItem];
-
-		NSButton* OKButton = [NSButton new];
-		[OKButton setButtonType:NSMomentaryLightButton];
-		[OKButton setTitle:@"OK"];
-		[OKButton setFrame:NSMakeRect(337, 13, 82, 32)];
-		[OKButton setImagePosition:NSNoImage];
-		[OKButton setBezelStyle:NSRoundedBezelStyle];
-		[OKButton setFont:[NSFont systemFontOfSize:0]];
-		[OKButton setAction:@selector(onOKButtonPressed:)];
-		[OKButton setTarget:self];
-
-		NSButton* CancelButton = [NSButton new];
-		[CancelButton setButtonType:NSMomentaryLightButton];
-		[CancelButton setTitle:@"Cancel"];
-		[CancelButton setFrame:NSMakeRect(255, 13, 82, 32)];
-		[CancelButton setImagePosition:NSNoImage];
-		[CancelButton setBezelStyle:NSRoundedBezelStyle];
-		[CancelButton setFont:[NSFont systemFontOfSize:0]];
-		[CancelButton setAction:@selector(onCancelButtonPressed:)];
-		[CancelButton setTarget:self];
-
-		[[Window contentView] addSubview:EditorMenu];
-		[[Window contentView] addSubview:OKButton];
-		[[Window contentView] addSubview:CancelButton];
-
-		[Window setTitle:@"Choose UE4 Editor..."];
-
-		[Window center];
-		[NSApp activateIgnoringOtherApps:YES];
-
-		NSInteger Result = [NSApp runModalForWindow:Window];
-		if (Result == NSOKButton)
+		FString EngineID = [self showEngineSelectionWindowForURL:FileURL];
+		if (!EngineID.IsEmpty() && !FDesktopPlatformModule::Get()->SetEngineIdentifierForProject([FileURL path], EngineID))
 		{
-			SelectedEnginePath = (NSString*)[[EditorMenu selectedItem] representedObject];
-			for (TMap<FString, FString>::TConstIterator Iter(Installations); Iter; ++Iter)
-			{
-				NSString* EnginePath = Iter.Value().GetNSString();
-				if ([EnginePath isEqualToString:SelectedEnginePath])
-				{
-					if (!FDesktopPlatformModule::Get()->SetEngineIdentifierForProject([FileURL path], Iter.Key()))
-					{
-						FPlatformMisc::MessageBoxExt(EAppMsgType::Ok, TEXT("Couldn't set association for project. Check the file is writeable."), TEXT("Error"));
-						return;
-					}
-					break;
-				}
-			}
+			FPlatformMisc::MessageBoxExt(EAppMsgType::Ok, TEXT("Couldn't set association for project. Check the file is writeable."), TEXT("Error"));
 		}
 	}
 }
