Improved the code for enumerating installed engine versions in UE4EditorServices so that it doesn't include incomplete or deleted installs

diff --git a/Engine/Source/Programs/Mac/UE4EditorServices/Resources/Mac/Info.plist b/Engine/Source/Programs/Mac/UE4EditorServices/Resources/Mac/Info.plist
--- a/Engine/Source/Programs/Mac/UE4EditorServices/Resources/Mac/Info.plist
+++ b/Engine/Source/Programs/Mac/UE4EditorServices/Resources/Mac/Info.plist
@@ -38,11 +38,11 @@
 	<key>CFBundlePackageType</key>
 	<string>APPL</string>
 	<key>CFBundleShortVersionString</key>
-	<string>1.5</string>
+	<string>1.6</string>
 	<key>CFBundleSignature</key>
 	<string>????</string>
 	<key>CFBundleVersion</key>
-	<string>1.5</string>
+	<string>1.6</string>
 	<key>LSApplicationCategoryType</key>
 	<string>public.app-category.developer-tools</string>
 	<key>LSBackgroundOnly</key>
