Linux: add stubbed file for UE4EditorServices.

- Gets built as a part of GUBP Tools node.

diff --git a/Engine/Source/Programs/Mac/UE4EditorServices/Private/Linux/UE4EditorServicesMain.cpp b/Engine/Source/Programs/Mac/UE4EditorServices/Private/Linux/UE4EditorServicesMain.cpp
new file mode 100644
--- /dev/null
+++ b/Engine/Source/Programs/Mac/UE4EditorServices/Private/Linux/UE4EditorServicesMain.cpp
@@ -0,0 +1,12 @@
+// Copyright 1998-2015 Epic Games, Inc. All Rights Reserved.
+
+#include "Core.h"
+#include "RequiredProgramMainCPPInclude.h"
+#include "ModuleManager.h"
+
+IMPLEMENT_APPLICATION(UE4EditorServices, "UE4EditorServices");
+
+int main(int argc, const char * argv[])
+{
+	return 0;
+}
