Stop UE4EditorServices from being built on Windows.

diff --git a/Engine/Source/Programs/Mac/UE4EditorServices/UE4EditorServices.Target.cs b/Engine/Source/Programs/Mac/UE4EditorServices/UE4EditorServices.Target.cs
--- a/Engine/Source/Programs/Mac/UE4EditorServices/UE4EditorServices.Target.cs
+++ b/Engine/Source/Programs/Mac/UE4EditorServices/UE4EditorServices.Target.cs
@@ -56,4 +56,15 @@ public class UE4EditorServicesTarget : TargetRules
 		CrossCompile = false;
 		return true;
 	}
+    public override List<UnrealTargetPlatform> GUBP_ToolPlatforms(UnrealTargetPlatform InHostPlatform)
+    {
+		if(InHostPlatform == UnrealTargetPlatform.Mac)
+		{
+			return new List<UnrealTargetPlatform> { InHostPlatform };
+		}
+		else
+		{
+			return new List<UnrealTargetPlatform>();
+		}
+    }
 }
