BootstrapPackagedGame for Mac

diff --git a/Engine/Source/Programs/Mac/BootstrapPackagedGame/BootstrapPackagedGame/Info.plist b/Engine/Source/Programs/Mac/BootstrapPackagedGame/BootstrapPackagedGame/Info.plist
new file mode 100644
--- /dev/null
+++ b/Engine/Source/Programs/Mac/BootstrapPackagedGame/BootstrapPackagedGame/Info.plist
@@ -0,0 +1,36 @@
+<?xml version="1.0" encoding="UTF-8"?>
+<!DOCTYPE plist PUBLIC "-//Apple//DTD PLIST 1.0//EN" "http://www.apple.com/DTDs/PropertyList-1.0.dtd">
+<plist version="1.0">
+<dict>
+	<key>CFBundleDevelopmentRegion</key>
+	<string>en</string>
+	<key>CFBundleExecutable</key>
+	<string>$(EXECUTABLE_NAME)</string>
+	<key>CFBundleIconFile</key>
+	<string>__UE4_ICON_FILE__</string>
+	<key>CFBundleIdentifier</key>
+	<string>com.epicgames.$(PRODUCT_NAME:rfc1034identifier)</string>
+	<key>CFBundleInfoDictionaryVersion</key>
+	<string>6.0</string>
+	<key>CFBundleName</key>
+	<string>$(PRODUCT_NAME)</string>
+	<key>CFBundlePackageType</key>
+	<string>APPL</string>
+	<key>CFBundleShortVersionString</key>
+	<string>1.0</string>
+	<key>CFBundleSignature</key>
+	<string>????</string>
+	<key>CFBundleVersion</key>
+	<string>1</string>
+	<key>LSMinimumSystemVersion</key>
+	<string>$(MACOSX_DEPLOYMENT_TARGET)</string>
+	<key>LSUIElement</key>
+	<true/>
+	<key>UE4AppToLaunch</key>
+	<string>__UE4_APP_TO_LAUNCH__</string>
+	<key>UE4CommandLine</key>
+	<string>__UE4_COMMANDLINE__</string>
+	<key>NSHumanReadableCopyright</key>
+	<string>Copyright © 2015 Epic Games, Inc. All rights reserved.</string>
+</dict>
+</plist>
diff --git a/Engine/Source/Programs/Mac/BootstrapPackagedGame/BootstrapPackagedGame/main.m b/Engine/Source/Programs/Mac/BootstrapPackagedGame/BootstrapPackagedGame/main.m
new file mode 100644
--- /dev/null
+++ b/Engine/Source/Programs/Mac/BootstrapPackagedGame/BootstrapPackagedGame/main.m
@@ -0,0 +1,117 @@
+// Copyright 1998-2015 Epic Games, Inc. All Rights Reserved.
+
+#import <Cocoa/Cocoa.h>
+
+int main(int argc, const char * argv[])
+{
+	NSWorkspace* Workspace = [NSWorkspace sharedWorkspace];
+	NSBundle* Bundle = [NSBundle mainBundle];
+	NSDictionary* Info = [Bundle infoDictionary];
+	NSString* AppToLaunch = [Info valueForKey:@"UE4AppToLaunch"];
+	NSString* CommandLine = [Info valueForKey:@"UE4CommandLine"];
+
+	if (!AppToLaunch)
+	{
+		AppToLaunch = [[[Bundle bundlePath] lastPathComponent] stringByDeletingPathExtension];
+	}
+
+	if ([AppToLaunch hasPrefix:@"com."])
+	{
+		AppToLaunch = [Workspace absolutePathForAppBundleWithIdentifier:AppToLaunch];
+	}
+	else if (![AppToLaunch hasPrefix:@"/"])
+	{
+		NSString* BaseDir = [[Bundle bundlePath] stringByDeletingLastPathComponent];
+		NSString* FullPath = [BaseDir stringByAppendingPathComponent:AppToLaunch];
+		NSFileManager* FileManager = [NSFileManager defaultManager];
+		if ([FileManager fileExistsAtPath:FullPath])
+		{
+			AppToLaunch = FullPath;
+		}
+		else
+		{
+			FullPath = [BaseDir stringByAppendingPathComponent:[NSString stringWithFormat:@"%@/Binaries/Mac/%@.app", AppToLaunch, AppToLaunch]];
+			if ([FileManager fileExistsAtPath:FullPath])
+			{
+				AppToLaunch = FullPath;
+			}
+			else
+			{
+				NSAlert *Alert = [[NSAlert alloc] init];
+				Alert.messageText = @"Couldn't find the app bundle to launch.";
+				[Alert setInformativeText:AppToLaunch];
+				[Alert runModal];
+				[Alert release];
+				return 0;
+			}
+		}
+	}
+
+	if (AppToLaunch)
+	{
+		NSMutableArray *Arguments = [[NSMutableArray new] autorelease];
+
+		if (CommandLine && [CommandLine length] > 0)
+		{
+			NSArray* ArgsArray = [CommandLine componentsSeparatedByString:@" "];
+
+			NSString* MultiPartArg = @"";
+			for (NSString* Item in ArgsArray)
+			{
+				if ([MultiPartArg length] == 0)
+				{
+					if (([Item hasPrefix:@"\""] && ![Item hasSuffix:@"\""]) // check for a starting quote but no ending quote, excludes quoted single arguments
+						|| ([Item rangeOfString:@"=\""].location != NSNotFound && ![Item hasSuffix:@"\""]) // check for quote after =, but no ending quote, this gets arguments of the type -blah="string string string"
+						|| [Item hasSuffix:@"=\""]) // check for ending quote after =, this gets arguments of the type -blah=" string string string "
+					{
+						MultiPartArg = [NSString stringWithString:Item];
+					}
+					else
+					{
+						NSString* Arg;
+						if ([Item rangeOfString:@"=\""].location != NSNotFound)
+						{
+							Arg = [[Item stringByReplacingOccurrencesOfString:@"=\"" withString:@"="] stringByTrimmingCharactersInSet:[NSCharacterSet characterSetWithCharactersInString:@"\""]];
+						}
+						else
+						{
+							Arg = [Item stringByTrimmingCharactersInSet:[NSCharacterSet characterSetWithCharactersInString:@"\""]];
+						}
+						[Arguments addObject:Arg];
+					}
+				}
+				else
+				{
+					MultiPartArg = [MultiPartArg stringByAppendingString:@" "];
+					MultiPartArg = [MultiPartArg stringByAppendingString:Item];
+					if ([Item hasSuffix:@"\""])
+					{
+						NSString* Arg;
+						if ([MultiPartArg hasPrefix:@"\""])
+						{
+							Arg = [MultiPartArg stringByTrimmingCharactersInSet:[NSCharacterSet characterSetWithCharactersInString:@"\""]];
+						}
+						else
+						{
+							Arg = [MultiPartArg stringByReplacingOccurrencesOfString:@"\"" withString:@""];
+						}
+						[Arguments addObject:Arg];
+						MultiPartArg = @"";
+					}
+				}
+			}
+		}
+
+		for (int Index = 1; Index < argc; Index++)
+		{
+			[Arguments addObject:[NSString stringWithUTF8String:argv[Index]]];
+		}
+
+		[Workspace launchApplicationAtURL:[NSURL fileURLWithPath:AppToLaunch]
+								  options:NSWorkspaceLaunchDefault
+							configuration:[NSDictionary dictionaryWithObject:Arguments forKey:NSWorkspaceLaunchConfigurationArguments]
+									error:nil];
+	}
+
+	return 0;
+}
