UE4EditorServices fixes:
- Fixed a mismatch in engine ID GUID formatting between ID stored in uproject file and engine installation ini on Mac
- Added an error message when user tries to generate Xcode project for blueprint-only project
- Updated version number to 1.7

diff --git a/Engine/Source/Developer/DesktopPlatform/Private/Mac/DesktopPlatformMac.cpp b/Engine/Source/Developer/DesktopPlatform/Private/Mac/DesktopPlatformMac.cpp
--- a/Engine/Source/Developer/DesktopPlatform/Private/Mac/DesktopPlatformMac.cpp
+++ b/Engine/Source/Developer/DesktopPlatform/Private/Mac/DesktopPlatformMac.cpp
@@ -552,7 +552,7 @@ bool FDesktopPlatformMac::RegisterEngineInstallation(const FString &RootDir, FSt
 		ConfigFile.Read(ConfigPath);
 
 		FConfigSection &Section = ConfigFile.FindOrAdd(TEXT("Installations"));
-		OutIdentifier = FGuid::NewGuid().ToString(EGuidFormats::DigitsWithHyphens);
+		OutIdentifier = FGuid::NewGuid().ToString(EGuidFormats::DigitsWithHyphensInBraces);
 		Section.AddUnique(*OutIdentifier, RootDir);
 
 		ConfigFile.Dirty = true;
@@ -615,7 +615,9 @@ void FDesktopPlatformMac::EnumerateEngineInstallations(TMap<FString, FString> &O
 				const FName* Key = Section.FindKey(EngineDir);
 				if (Key)
 				{
-					EngineId = Key->ToString();
+					FGuid IdGuid;
+					FGuid::Parse(Key->ToString(), IdGuid);
+					EngineId = IdGuid.ToString(EGuidFormats::DigitsWithHyphensInBraces);;
 				}
 				else
 				{
diff --git a/Engine/Source/Programs/Mac/UE4EditorServices/Private/Mac/UE4EditorServicesAppDelegate.cpp b/Engine/Source/Programs/Mac/UE4EditorServices/Private/Mac/UE4EditorServicesAppDelegate.cpp
--- a/Engine/Source/Programs/Mac/UE4EditorServices/Private/Mac/UE4EditorServicesAppDelegate.cpp
+++ b/Engine/Source/Programs/Mac/UE4EditorServices/Private/Mac/UE4EditorServicesAppDelegate.cpp
@@ -298,6 +298,14 @@
 		NSString* EnginePath = [self findEngineForUProjectFile:FileURL];
 		if (EnginePath)
 		{
+			// Check it's a code project
+			FString SourceDir = FPaths::GetPath(FString([[PBoard propertyListForType:NSFilenamesPboardType] objectAtIndex:0])) / TEXT("Source");
+			if(!IPlatformFile::GetPlatformPhysical().DirectoryExists(*SourceDir))
+			{
+				FPlatformMisc::MessageBoxExt(EAppMsgType::Ok, TEXT("This project does not have any source code. You need to add C++ source files to the project from the Editor before you can generate project files."), TEXT("Error"));
+				return;
+			}
+
 			NSString* ScriptPath = [EnginePath stringByAppendingPathComponent:@"Engine/Build/BatchFiles/Mac/GenerateProjectFiles.sh"];
 			if (![[NSFileManager defaultManager] fileExistsAtPath:ScriptPath])
 			{
diff --git a/Engine/Source/Programs/Mac/UE4EditorServices/Resources/Mac/Info.plist b/Engine/Source/Programs/Mac/UE4EditorServices/Resources/Mac/Info.plist
--- a/Engine/Source/Programs/Mac/UE4EditorServices/Resources/Mac/Info.plist
+++ b/Engine/Source/Programs/Mac/UE4EditorServices/Resources/Mac/Info.plist
@@ -38,11 +38,11 @@
 	<key>CFBundlePackageType</key>
 	<string>APPL</string>
 	<key>CFBundleShortVersionString</key>
-	<string>1.6</string>
+	<string>1.7</string>
 	<key>CFBundleSignature</key>
 	<string>????</string>
 	<key>CFBundleVersion</key>
-	<string>1.6</string>
+	<string>1.7</string>
 	<key>LSApplicationCategoryType</key>
 	<string>public.app-category.developer-tools</string>
 	<key>LSBackgroundOnly</key>
