Build: Fixes to allow compiling the whole solution, rather than just certain targets.

* TestPAL: Compile failure due to an MSVC warning about assigning 32-bit shift result to 64-bit variable. Also added to CIS.
* UE4EditorServices: Conflict in variable name for hInstance. Disabled building this target by default for Windows and Linux, since it only usefully does anything on Mac.

diff --git a/Engine/Source/Programs/Mac/UE4EditorServices/UE4EditorServices.Target.cs b/Engine/Source/Programs/Mac/UE4EditorServices/UE4EditorServices.Target.cs
--- a/Engine/Source/Programs/Mac/UE4EditorServices/UE4EditorServices.Target.cs
+++ b/Engine/Source/Programs/Mac/UE4EditorServices/UE4EditorServices.Target.cs
@@ -25,6 +25,12 @@ public class UE4EditorServicesTarget : TargetRules
 			);
 	}
 
+	public override bool GetSupportedPlatforms(ref List<UnrealTargetPlatform> OutPlatforms)
+	{
+		OutPlatforms.Add( UnrealTargetPlatform.Mac );
+		return true;
+	}
+
 	public override bool ShouldCompileMonolithic(UnrealTargetPlatform InPlatform, UnrealTargetConfiguration InConfiguration)
 	{
 		return true;
diff --git a/Engine/Source/Programs/TestPAL/Private/Main.cpp b/Engine/Source/Programs/TestPAL/Private/Main.cpp
--- a/Engine/Source/Programs/TestPAL/Private/Main.cpp
+++ b/Engine/Source/Programs/TestPAL/Private/Main.cpp
@@ -458,7 +458,7 @@ int32 GetAllocationSizeTest(const TCHAR* CommandLine)
 	{
 		for (SIZE_T AlignmentPower = 4; AlignmentPower <= 7; ++AlignmentPower)
 		{
-			SIZE_T Alignment = (1 << AlignmentPower);
+			SIZE_T Alignment = ((SIZE_T)1 << AlignmentPower);
 
 			Allocation New;
 			New.RequestedSize = Size;
