Fix 'Generate Xcode Project (UnrealEditorServices)'

This adds the ability to request permission to use apple events to control the terminal. This fixes an issue where UnrealEditorServices would not generate Xcode projects and just launch a blank terminal.

diff --git a/Engine/Source/Programs/Mac/UnrealEditorServices/Private/Mac/UnrealEditorServicesAppDelegate.cpp b/Engine/Source/Programs/Mac/UnrealEditorServices/Private/Mac/UnrealEditorServicesAppDelegate.cpp
--- a/Engine/Source/Programs/Mac/UnrealEditorServices/Private/Mac/UnrealEditorServicesAppDelegate.cpp
+++ b/Engine/Source/Programs/Mac/UnrealEditorServices/Private/Mac/UnrealEditorServicesAppDelegate.cpp
@@ -287,6 +287,14 @@
 	}
 }
 
+- (BOOL)canControlTerminal
+{
+	// This will verify if we can access the terminal and ask the user if necessary
+	NSAppleEventDescriptor* targetAppEventDescriptor = [NSAppleEventDescriptor descriptorWithBundleIdentifier:@"com.apple.terminal"];
+	OSStatus status = AEDeterminePermissionToAutomateTarget(targetAppEventDescriptor.aeDesc, typeWildCard, typeWildCard, true);
+	return (status == noErr);
+}
+
 - (void)generateXcodeProjectService:(NSPasteboard *)PBoard userData:(NSString *)UserData error:(NSString **)Error
 {
 	NSString* UnrealError = NULL;
@@ -315,37 +323,44 @@
 			{
 				if ([[NSWorkspace sharedWorkspace] launchAppWithBundleIdentifier:@"com.apple.terminal" options:0 additionalEventParamDescriptor:nil launchIdentifier:nil])
 				{
-					NSString* FullFolderPath = [[ScriptPath stringByDeletingLastPathComponent] stringByResolvingSymlinksInPath];
-					NSString* FullScriptPath = [ScriptPath stringByResolvingSymlinksInPath];
-					NSString* Command = [NSString stringWithFormat:@"cd \"%@\" \n sh \"%@\" -project=\"%@\" -game\n logout\n", FullFolderPath, FullScriptPath, [FileURL path]];
-
-					const char* UTF8Script = [Command UTF8String];
-
-					const char *TerminalBundleID = "com.apple.terminal";
-
-					// Build event
-					AppleEvent Event, Reply;
-					OSStatus Status = AEBuildAppleEvent(kAECoreSuite, kAEDoScript,
-											typeApplicationBundleID,
-											TerminalBundleID, strlen(TerminalBundleID),
-											kAutoGenerateReturnID,
-											kAnyTransactionID,
-											&Event, NULL,
-											"'----':utf8(@)", [Command length],
-											UTF8Script);
-					if (Status == noErr)
+					if ( [self canControlTerminal] )
 					{
-						// Send event and check for any Apple Event Manager errors
-						Status = AESendMessage(&Event, &Reply, kAENoReply, kAEDefaultTimeout);
-						AEDisposeDesc(&Event);
-						if (Status != noErr)
+						NSString* FullFolderPath = [[ScriptPath stringByDeletingLastPathComponent] stringByResolvingSymlinksInPath];
+						NSString* FullScriptPath = [ScriptPath stringByResolvingSymlinksInPath];
+						NSString* Command = [NSString stringWithFormat:@"cd \"%@\" \n sh \"%@\" -project=\"%@\" -game\n logout\n", FullFolderPath, FullScriptPath, [FileURL path]];
+
+						const char* UTF8Script = [Command UTF8String];
+
+						const char *TerminalBundleID = "com.apple.terminal";
+
+						// Build event
+						AppleEvent Event, Reply;
+						OSStatus Status = AEBuildAppleEvent(kAECoreSuite, kAEDoScript,
+												typeApplicationBundleID,
+												TerminalBundleID, strlen(TerminalBundleID),
+												kAutoGenerateReturnID,
+												kAnyTransactionID,
+												&Event, NULL,
+												"'----':utf8(@)", [Command length],
+												UTF8Script);
+						if (Status == noErr)
+						{
+							// Send event and check for any Apple Event Manager errors
+							Status = AESendMessage(&Event, &Reply, kAENoReply, kAEDefaultTimeout);
+							AEDisposeDesc(&Event);
+							if (Status != noErr)
+							{
+								UnrealError = @"Couldn't tell Terminal to generate project files.";
+							}
+						}
+						else
 						{
 							UnrealError = @"Couldn't tell Terminal to generate project files.";
 						}
 					}
 					else
 					{
-						UnrealError = @"Couldn't tell Terminal to generate project files.";
+						UnrealError = @"Didn't have permission to tell Terminal to generate project files.";
 					}
 				}
 				else
